Fix dependency chain while including PIEPreviewDeviceSpecification

diff --git a/Engine/Source/Developer/Android/AndroidDeviceDetection/AndroidDeviceDetection.Build.cs b/Engine/Source/Developer/Android/AndroidDeviceDetection/AndroidDeviceDetection.Build.cs
--- a/Engine/Source/Developer/Android/AndroidDeviceDetection/AndroidDeviceDetection.Build.cs
+++ b/Engine/Source/Developer/Android/AndroidDeviceDetection/AndroidDeviceDetection.Build.cs
@@ -14,7 +14,7 @@ public class AndroidDeviceDetection : ModuleRules
 				"Core",
 				"Json",
                 "JsonUtilities",
-                "PIEPreviewDeviceProfileSelector"
+                "PIEPreviewDeviceSpecification"
             }
 		);
 
@@ -32,14 +32,6 @@ public class AndroidDeviceDetection : ModuleRules
 			}
 		);
 
-        PrivateIncludePaths.AddRange(
-            new string[]
-            {
-                "Editor/PIEPreviewDeviceProfileSelector/Private"
-            }
-        );
-
-
         if (Target.bCompileAgainstEngine)
 		{
 			PrivateDependencyModuleNames.Add("Engine");
diff --git a/Engine/Source/Editor/PIEPreviewDeviceProfileSelector/PIEPreviewDeviceProfileSelector.Build.cs b/Engine/Source/Editor/PIEPreviewDeviceProfileSelector/PIEPreviewDeviceProfileSelector.Build.cs
--- a/Engine/Source/Editor/PIEPreviewDeviceProfileSelector/PIEPreviewDeviceProfileSelector.Build.cs
+++ b/Engine/Source/Editor/PIEPreviewDeviceProfileSelector/PIEPreviewDeviceProfileSelector.Build.cs
@@ -19,7 +19,8 @@ namespace UnrealBuildTool.Rules
 			PrivateDependencyModuleNames.AddRange(
 				new string[]
 				{
-					"RHI",
+                    "PIEPreviewDeviceSpecification",
+                    "RHI",
 					"Json",
 					"JsonUtilities",
 					"MaterialShaderQualitySettings",
diff --git a/Engine/Source/Editor/PIEPreviewDeviceProfileSelector/Private/PIEPreviewDeviceSpecification.cpp b/Engine/Source/Editor/PIEPreviewDeviceProfileSelector/Private/PIEPreviewDeviceSpecification.cpp
deleted file mode 100644
diff --git a/Engine/Source/Editor/PIEPreviewDeviceProfileSelector/Private/PIEPreviewDeviceSpecification.h b/Engine/Source/Editor/PIEPreviewDeviceProfileSelector/Private/PIEPreviewDeviceSpecification.h
deleted file mode 100644
diff --git a/Engine/Source/Editor/PIEPreviewDeviceSpecification/PIEPreviewDeviceSpecification.Build.cs b/Engine/Source/Editor/PIEPreviewDeviceSpecification/PIEPreviewDeviceSpecification.Build.cs
new file mode 100644
--- /dev/null
+++ b/Engine/Source/Editor/PIEPreviewDeviceSpecification/PIEPreviewDeviceSpecification.Build.cs
@@ -0,0 +1,20 @@
+// Copyright 1998-2018 Epic Games, Inc. All Rights Reserved.
+
+namespace UnrealBuildTool.Rules
+{
+    public class PIEPreviewDeviceSpecification : ModuleRules
+	{
+        public PIEPreviewDeviceSpecification(ReadOnlyTargetRules Target) : base(Target)
+		{
+			PrivateIncludePathModuleNames.Add("Engine");
+
+			PublicDependencyModuleNames.AddRange(
+				new string[]
+				{
+					"Core",
+				    "CoreUObject",
+                }
+				);
+		}
+	}
+}
diff --git a/Engine/Source/Editor/PIEPreviewDeviceSpecification/Private/PIEPreviewDeviceSpecification.cpp b/Engine/Source/Editor/PIEPreviewDeviceSpecification/Private/PIEPreviewDeviceSpecification.cpp
new file mode 100644
--- /dev/null
+++ b/Engine/Source/Editor/PIEPreviewDeviceSpecification/Private/PIEPreviewDeviceSpecification.cpp
@@ -0,0 +1,24 @@
+// Copyright 1998-2018 Epic Games, Inc. All Rights Reserved.
+
+#include "PIEPreviewDeviceSpecification.h"
+#include "Modules/ModuleInterface.h"
+#include "Modules/ModuleManager.h"
+
+class FPIEPreviewDeviceSpecificationModule
+	: public IModuleInterface
+{
+public:
+
+	// IModuleInterface interface
+
+	virtual void StartupModule() override { }
+	virtual void ShutdownModule() override { }
+};
+
+
+IMPLEMENT_MODULE(FPIEPreviewDeviceSpecificationModule, PIEPreviewDeviceSpecification);
+
+UPIEPreviewDeviceSpecification::UPIEPreviewDeviceSpecification(const FObjectInitializer& ObjectInitializer)
+	: Super(ObjectInitializer)
+{
+}
diff --git a/Engine/Source/Editor/PIEPreviewDeviceSpecification/Public/PIEPreviewDeviceSpecification.h b/Engine/Source/Editor/PIEPreviewDeviceSpecification/Public/PIEPreviewDeviceSpecification.h
new file mode 100644
--- /dev/null
+++ b/Engine/Source/Editor/PIEPreviewDeviceSpecification/Public/PIEPreviewDeviceSpecification.h
@@ -0,0 +1,174 @@
+// Copyright 1998-2018 Epic Games, Inc. All Rights Reserved.
+
+#pragma once
+
+#include "CoreMinimal.h"
+#include "UObject/Object.h"
+
+#include "PIEPreviewDeviceSpecification.generated.h"
+
+UENUM()
+enum class EPIEPreviewDeviceType : uint8
+{
+	Unset,
+	Android,
+	IOS,
+	TVOS,
+	MAX,
+};
+
+UCLASS()
+class PIEPREVIEWDEVICESPECIFICATION_API UPIEPreviewDeviceSpecification : public UObject
+{
+	GENERATED_UCLASS_BODY()
+public:
+	UPROPERTY()
+		EPIEPreviewDeviceType PreviewDeviceType;
+
+	UPROPERTY()
+		FString GPUFamily;
+	UPROPERTY()
+		FString GLVersion;
+	UPROPERTY()
+		FString VulkanVersion;
+	UPROPERTY()
+		FString AndroidVersion;
+	UPROPERTY()
+		FString DeviceMake;
+	UPROPERTY()
+		FString DeviceModel;
+	UPROPERTY()
+		bool UsingHoudini;
+};
+
+USTRUCT()
+struct FPIERHIOverrideState
+{
+public:
+	GENERATED_USTRUCT_BODY()
+	UPROPERTY()
+	int32 MaxShadowDepthBufferSizeX;
+	UPROPERTY()
+	int32 MaxShadowDepthBufferSizeY;
+	UPROPERTY()
+	int32 MaxTextureDimensions;
+	UPROPERTY()
+	int32 MaxCubeTextureDimensions;
+	UPROPERTY()
+	bool SupportsRenderTargetFormat_PF_G8;
+	UPROPERTY()
+	bool SupportsRenderTargetFormat_PF_FloatRGBA;
+	UPROPERTY()
+	bool SupportsMultipleRenderTargets;
+	UPROPERTY()
+	bool SupportsInstancing;
+};
+
+USTRUCT()
+struct FPIEAndroidDeviceProperties
+{
+public:
+	GENERATED_USTRUCT_BODY()
+
+	UPROPERTY()
+	FString GPUFamily;
+	UPROPERTY()
+	FString GLVersion;
+	UPROPERTY()
+	FString VulkanVersion;
+	UPROPERTY()
+	FString AndroidVersion;
+	UPROPERTY()
+	FString DeviceMake;
+	UPROPERTY()
+	FString DeviceModel;
+	UPROPERTY()
+	bool VulkanAvailable;
+	UPROPERTY()
+	bool UsingHoudini;
+
+	UPROPERTY()
+	FPIERHIOverrideState GLES2RHIState;
+
+	UPROPERTY()
+	FPIERHIOverrideState GLES31RHIState;
+
+// 	UPROPERTY()
+// 	FPIERHIOverrideState VulkanRHIState;
+};
+
+USTRUCT()
+struct FPIEIOSDeviceProperties
+{
+public:
+	GENERATED_USTRUCT_BODY()
+
+	UPROPERTY()
+	FString DeviceModel;
+
+	UPROPERTY()
+	float NativeScaleFactor = 0.0f;
+
+	UPROPERTY()
+	FPIERHIOverrideState GLES2RHIState;
+
+	UPROPERTY()
+	FPIERHIOverrideState MetalRHIState;
+};
+
+USTRUCT()
+struct FPIEPreviewDeviceBezelViewportRect
+{
+public:
+	GENERATED_USTRUCT_BODY()
+	
+	UPROPERTY()
+	int32 X = 0;
+	UPROPERTY()
+	int32 Y = 0;
+	UPROPERTY()
+	int32 Width = 0;
+	UPROPERTY()
+	int32 Height = 0;
+}; 
+
+USTRUCT()
+struct FPIEBezelProperties
+{
+	GENERATED_USTRUCT_BODY()
+
+	UPROPERTY()
+	FString DeviceBezelFile;
+	UPROPERTY()
+	FPIEPreviewDeviceBezelViewportRect BezelViewportRect;
+};
+
+USTRUCT()
+struct FPIEPreviewDeviceSpecifications
+{
+	GENERATED_USTRUCT_BODY()
+	
+	UPROPERTY()
+	EPIEPreviewDeviceType DevicePlatform;
+	UPROPERTY()
+	int32 ResolutionX;
+	UPROPERTY()
+	int32 ResolutionY;
+	UPROPERTY()
+	int32 ResolutionYImmersiveMode;
+
+	UPROPERTY()
+	int32 PPI;
+
+	UPROPERTY()
+	TArray<float> ScaleFactors;
+	
+	UPROPERTY()
+	FPIEBezelProperties BezelProperties;
+
+	UPROPERTY()
+	FPIEAndroidDeviceProperties AndroidProperties;
+
+	UPROPERTY()
+	FPIEIOSDeviceProperties IOSProperties;
+};
