Fix static analysis

#Synced-CL 50957382

diff --git a/Engine/Source/Editor/PIEPreviewDeviceSpecification/Public/PIEPreviewDeviceSpecification.h b/Engine/Source/Editor/PIEPreviewDeviceSpecification/Public/PIEPreviewDeviceSpecification.h
--- a/Engine/Source/Editor/PIEPreviewDeviceSpecification/Public/PIEPreviewDeviceSpecification.h
+++ b/Engine/Source/Editor/PIEPreviewDeviceSpecification/Public/PIEPreviewDeviceSpecification.h
@@ -85,7 +85,7 @@ struct FPIEAndroidDeviceProperties
 public:
 	GENERATED_USTRUCT_BODY()
 	UPROPERTY()
-	uint32 Version = 0;
+	int32 Version = 0;
 	UPROPERTY()
 	FString GPUFamily;
 	UPROPERTY()
@@ -128,7 +128,7 @@ public:
 	GENERATED_USTRUCT_BODY()
 	
 	UPROPERTY()
-	uint32 Version = 0;
+	int32 Version = 0;
 
 	UPROPERTY()
 	FString DeviceModel;
